Switchboard: Added helper to lock gpu clocks, which is currently required for reliable smpte-2110 gpu direct.

diff --git a/Engine/Build/InstalledEngineBuild.xml b/Engine/Build/InstalledEngineBuild.xml
--- a/Engine/Build/InstalledEngineBuild.xml
+++ b/Engine/Build/InstalledEngineBuild.xml
@@ -476,6 +476,7 @@
 			<Compile Target="BootstrapPackagedGame" Configuration="Shipping" Platform="Win64" Tag="#Build Tools Win64"/>
 			<Compile Target="BuildPatchTool" Configuration="Shipping" Platform="Win64" Tag="#Build Tools Win64"/>
 			<Compile Target="SwitchboardListener" Configuration="Development" Platform="Win64" Tag="#Build Tools Win64"/>
+			<Compile Target="SwitchboardListenerHelper" Configuration="Development" Platform="Win64" Tag="#Build Tools Win64"/>
 			<Compile Target="UnrealObjectPtrTool" Configuration="Development" Platform="Win64" Tag="#Build Tools Win64"/>
 
 			<Tag Files="Engine/Binaries/Win64/XGEControlWorker.exe" With="#Build Tools Win64"/>
@@ -552,6 +553,7 @@
 			<Compile Target="BuildPatchTool" Configuration="Shipping" Platform="Linux" Tag="#Build Tools Linux"/>
 			<Compile Target="UnrealRecoverySvc" Configuration="Development" Platform="Linux" Tag="#Build Tools Linux"/>
 			<Compile Target="SwitchboardListener" Configuration="Development" Platform="Linux" Tag="#Build Tools Linux"/>
+			<Compile Target="SwitchboardListenerHelper" Configuration="Development" Platform="Linux" Tag="#Build Tools Linux"/>
 			<Compile Target="UnrealObjectPtrTool" Configuration="Development" Platform="Linux" Tag="#Build Tools Linux"/>
 		</Node>
 	</Agent>
@@ -642,6 +644,7 @@
 			<Compile Target="BuildPatchTool" Configuration="Shipping" Platform="Linux" Tag="#Build Tools LinuxNative"/>
 			<Compile Target="UnrealRecoverySvc" Configuration="Development" Platform="Linux" Tag="#Build Tools LinuxNative"/>
 			<Compile Target="SwitchboardListener" Configuration="Development" Platform="Linux" Tag="#Build Tools LinuxNative"/>
+			<Compile Target="SwitchboardListenerHelper" Configuration="Development" Platform="Linux" Tag="#Build Tools LinuxNative"/>
 			<Compile Target="UnrealObjectPtrTool" Configuration="Development" Platform="Linux" Tag="#Build Tools LinuxNative"/>
 		</Node>
 
diff --git a/Engine/Plugins/VirtualProduction/Switchboard/Source/Switchboard/switchboard/config.py b/Engine/Plugins/VirtualProduction/Switchboard/Source/Switchboard/switchboard/config.py
--- a/Engine/Plugins/VirtualProduction/Switchboard/Source/Switchboard/switchboard/config.py
+++ b/Engine/Plugins/VirtualProduction/Switchboard/Source/Switchboard/switchboard/config.py
@@ -1961,6 +1961,7 @@ class Config(object):
             data = {}
 
         self.init_switchboard_settings()
+        self.init_sblhelper_settings()
         self.init_project_settings(data)
         self.init_unreal_insights(data)
         self.init_muserver(data)
@@ -2041,6 +2042,7 @@ class Config(object):
 
         self.file_path = get_absolute_config_path(file_path)
         self.init_switchboard_settings()
+        self.init_sblhelper_settings()
         self.init_project_settings(
             { 
                 "project_name": self.file_path.stem, 
@@ -2066,14 +2068,27 @@ class Config(object):
     def init_switchboard_settings(self, data={}):
         self.switchboard_settings = {
             "listener_exe": StringSetting(
-                "listener_exe",
-                "Listener Executable Name",
-                 data.get('listener_exe', 'SwitchboardListener')
+                attr_name = "listener_exe",
+                nice_name = "Listener Executable Name",
+                value = data.get('listener_exe', 'SwitchboardListener')
             )
         }
         
         self.LISTENER_EXE = self.switchboard_settings["listener_exe"]
 
+    def init_sblhelper_settings(self, data={}):
+        self.sblhelper_settings = {
+            "sblhelper_exe": StringSetting(
+                attr_name = "sblhelper_exe",
+                nice_name = "Gpu Clocker Executable Name",
+                value = data.get('sblhelper_exe', 'SwitchboardListenerHelper'),
+                tool_tip = "Name of the executable that SwitchboardListener can communicate with to lock Gpu clocks.",
+                show_ui = True if sys.platform in ('win32','linux') else False, # # Gpu Clocker is available in select platforms
+            )
+        }
+        
+        self.SBLHELPER_EXE = self.sblhelper_settings["sblhelper_exe"]
+
     def init_project_settings(self, data={}):
         self.basic_project_settings = {
             "project_name": StringSetting(
@@ -2439,6 +2454,7 @@ class Config(object):
         data["content_plugin_filters"] = (
             self.CONTENT_PLUGIN_FILTERS.get_value())
         data["listener_exe"] = self.LISTENER_EXE.get_value()
+        data["sblhelper_exe"] = self.SBLHELPER_EXE.get_value()
 
         self.save_unreal_insights(data)
 
@@ -2672,6 +2688,10 @@ class Config(object):
         return self.engine_exe_path(
             self.ENGINE_DIR.get_value(), self.LISTENER_EXE.get_value())
 
+    def sblhelper_path(self):
+        return self.engine_exe_path(
+            self.ENGINE_DIR.get_value(), self.SBLHELPER_EXE.get_value())
+
     # todo-dara: find a way to do this directly in the LiveLinkFace plugin code
     def unreal_device_addresses(self):
         unreal_addresses = []
diff --git a/Engine/Plugins/VirtualProduction/Switchboard/Source/Switchboard/switchboard/devices/ndisplay/plugin_ndisplay.py b/Engine/Plugins/VirtualProduction/Switchboard/Source/Switchboard/switchboard/devices/ndisplay/plugin_ndisplay.py
--- a/Engine/Plugins/VirtualProduction/Switchboard/Source/Switchboard/switchboard/devices/ndisplay/plugin_ndisplay.py
+++ b/Engine/Plugins/VirtualProduction/Switchboard/Source/Switchboard/switchboard/devices/ndisplay/plugin_ndisplay.py
@@ -3,6 +3,7 @@
 import concurrent.futures
 import json
 import os
+import sys
 from pathlib import Path
 import socket
 import struct
@@ -490,6 +491,16 @@ class DevicenDisplay(DeviceUnreal):
             value='',
             tool_tip=('Adds the selected Media Profile to the command line')
         ),
+        'lock_gpu_clock': BoolSetting(
+            attr_name="lock_gpu_clock",
+            nice_name="Lock GPU Clock",
+            value=False,
+            tool_tip=(
+                "Hint to lock the GPU clock to its allowed maximum. Requires SwitchboardListenerHelper \n"
+                "to be running on the client machine, otherwise this option will be ignored."
+            ),
+            show_ui = True if sys.platform in ('win32','linux') else False, # Gpu Clocker is available in select platforms
+        ),
     }
 
     ndisplay_monitor_ui = None
diff --git a/Engine/Plugins/VirtualProduction/Switchboard/Source/Switchboard/switchboard/devices/unreal/plugin_unreal.py b/Engine/Plugins/VirtualProduction/Switchboard/Source/Switchboard/switchboard/devices/unreal/plugin_unreal.py
--- a/Engine/Plugins/VirtualProduction/Switchboard/Source/Switchboard/switchboard/devices/unreal/plugin_unreal.py
+++ b/Engine/Plugins/VirtualProduction/Switchboard/Source/Switchboard/switchboard/devices/unreal/plugin_unreal.py
@@ -667,6 +667,16 @@ class DeviceUnreal(Device):
             value='',
             tool_tip=('Adds the selected Media Profile to the command line')
         ),
+        'lock_gpu_clock': BoolSetting(
+            attr_name="lock_gpu_clock",
+            nice_name="Lock GPU Clock",
+            value=False,
+            tool_tip=(
+                "Hint to lock the GPU clock to its allowed maximum. Requires SwitchboardListenerHelper \n"
+                "to be running on the client machine, otherwise this option will be ignored."
+            ),
+            show_ui = True if sys.platform in ('win32','linux') else False, # Gpu Clocker is available in select platforms
+        ),
         'use_sync_filters': BoolSetting(
             attr_name='use_sync_filters',
             nice_name='Use Sync Filters',
@@ -2087,13 +2097,17 @@ class DeviceUnreal(Device):
 
         self.last_launch_command.update_value(f'{engine_path} {args}')
 
+        lock_gpu_clock_csetting = self.__class__.csettings.get('lock_gpu_clock', None)
+        lock_gpu_clock = lock_gpu_clock_csetting.get_value() if lock_gpu_clock_csetting else False
+        
         puuid, msg = message_protocol.create_start_process_message(
             prog_path=engine_path,
             prog_args=args,
             prog_name=program_name,
             caller=self.name,
             update_clients_with_stdout=False,
-            priority_modifier=priority_modifier
+            priority_modifier=priority_modifier,
+            lock_gpu_clock=lock_gpu_clock,
         )
 
         self.program_start_queue.add(
diff --git a/Engine/Plugins/VirtualProduction/Switchboard/Source/Switchboard/switchboard/devices/unreal/version_helpers.py b/Engine/Plugins/VirtualProduction/Switchboard/Source/Switchboard/switchboard/devices/unreal/version_helpers.py
--- a/Engine/Plugins/VirtualProduction/Switchboard/Source/Switchboard/switchboard/devices/unreal/version_helpers.py
+++ b/Engine/Plugins/VirtualProduction/Switchboard/Source/Switchboard/switchboard/devices/unreal/version_helpers.py
@@ -3,7 +3,7 @@
 import typing
 
 
-LISTENER_COMPATIBLE_VERSION = (2,2)
+LISTENER_COMPATIBLE_VERSION = (2,3)
 
 # Version where redeploy support was initially introduced.
 LISTENER_MIN_REDEPLOY_VERSION = (1,4,0)
diff --git a/Engine/Plugins/VirtualProduction/Switchboard/Source/Switchboard/switchboard/message_protocol.py b/Engine/Plugins/VirtualProduction/Switchboard/Source/Switchboard/switchboard/message_protocol.py
--- a/Engine/Plugins/VirtualProduction/Switchboard/Source/Switchboard/switchboard/message_protocol.py
+++ b/Engine/Plugins/VirtualProduction/Switchboard/Source/Switchboard/switchboard/message_protocol.py
@@ -14,6 +14,7 @@ def create_start_process_message(
     *,
 	update_clients_with_stdout: bool = False,
     priority_modifier: int = 0,
+    lock_gpu_clock: bool = False,
 ):
     cmd_id = uuid.uuid4()
     start_cmd = {
@@ -26,6 +27,7 @@ def create_start_process_message(
         'working_dir': working_dir,
         'bUpdateClientsWithStdout': update_clients_with_stdout,
         'priority_modifier': priority_modifier,
+        'bLockGpuClock' : lock_gpu_clock,
     }
 
     message = json.dumps(start_cmd).encode() + b'\x00'
diff --git a/Engine/Plugins/VirtualProduction/Switchboard/Source/Switchboard/switchboard/settings_dialog.py b/Engine/Plugins/VirtualProduction/Switchboard/Source/Switchboard/switchboard/settings_dialog.py
--- a/Engine/Plugins/VirtualProduction/Switchboard/Source/Switchboard/switchboard/settings_dialog.py
+++ b/Engine/Plugins/VirtualProduction/Switchboard/Source/Switchboard/switchboard/settings_dialog.py
@@ -208,6 +208,7 @@ class SettingsDialog(QtCore.QObject):
         settings.ADDRESS.create_ui(form_layout=layout)
         settings.TRANSPORT_PATH.create_ui(form_layout=layout)
         config.LISTENER_EXE.create_ui(form_layout=layout)
+        config.SBLHELPER_EXE.create_ui(form_layout=layout)
 
         return self.ui.switchboard_settings_group
         
diff --git a/Engine/Plugins/VirtualProduction/Switchboard/Source/Switchboard/switchboard/switchboard_dialog.py b/Engine/Plugins/VirtualProduction/Switchboard/Source/Switchboard/switchboard/switchboard_dialog.py
--- a/Engine/Plugins/VirtualProduction/Switchboard/Source/Switchboard/switchboard/switchboard_dialog.py
+++ b/Engine/Plugins/VirtualProduction/Switchboard/Source/Switchboard/switchboard/switchboard_dialog.py
@@ -8,6 +8,7 @@ import os
 import threading
 import time
 import re
+import sys
 from typing import Callable, List, Optional, Set, Union
 
 from pathlib import Path
@@ -37,6 +38,7 @@ from switchboard.settings_dialog import SettingsDialog
 from switchboard.switchboard_logging import ConsoleStream, LOGGER
 from switchboard.tools.insights_launcher import InsightsLauncher
 from switchboard.tools.listener_launcher import ListenerLauncher
+from switchboard.tools.sblhelper_launcher import SBLHelperLauncher
 from switchboard.devices.unreal.plugin_unreal import DeviceUnreal
 from switchboard.devices.unreal.redeploy_dialog import RedeployListenerDialog
 from switchboard.util import collect_logs
@@ -210,6 +212,7 @@ class ProcessMonitor(QtCore.QObject):
         while not self._exiting:
             self._poll_multiuser_status()
             self._dialog.update_locallistener_menuitem()
+            self._dialog.update_localsblhelper_menuitem()
             self._dialog.update_insights_menuitem()
 
             time.sleep(1.0)
@@ -309,9 +312,12 @@ class SwitchboardDialog(QtCore.QObject):
         # Convenience UnrealInsights launcher
         self.init_insights_launcher()
 
-        # Convenience local Switchboard lister launcher
+        # Convenience local Switchboard Listener launcher
         self.init_listener_launcher()
 
+        # Convenience local Switchboard Listener Helper launcher
+        self.init_sblhelper_launcher()
+
         # Convenience Open Logs Folder menu item
         self.register_open_logs_menuitem()
         self.register_zip_logs_menuitem()
@@ -577,6 +583,13 @@ class SwitchboardDialog(QtCore.QObject):
         '''
         self.locallistener_launcher_menuitem.setEnabled(not self.listener_launcher.is_running())
 
+    def update_localsblhelper_menuitem(self):
+        ''' 
+        Enables/disables the local listener helper launch menu item depending on whether 
+        it is already running or not.
+        '''
+        self.localsblhelper_launcher_menuitem.setEnabled(not self.sblhelper_launcher.is_running())
+
     def update_insights_menuitem(self):
         ''' 
         Enables/disables the UnrealInsights launch menu item depending on whether 
@@ -614,7 +627,7 @@ class SwitchboardDialog(QtCore.QObject):
         self.insights_launcher_menuitem = action
 
     def init_listener_launcher(self):
-        ''' Initializes switcboard listener launcher '''
+        ''' Initializes switchboard listener launcher '''
         self.listener_launcher = ListenerLauncher()
 
         def launch_listener():
@@ -628,6 +641,25 @@ class SwitchboardDialog(QtCore.QObject):
 
         self.locallistener_launcher_menuitem = action
 
+    def init_sblhelper_launcher(self):
+        ''' Initializes switchboard listener helper launcher '''
+        self.sblhelper_launcher = SBLHelperLauncher()
+
+        def launch_sblhelper():
+            try:
+                self.sblhelper_launcher.launch()
+            except Exception as e:
+                LOGGER.error(e)
+
+        # Gpu Clocker is available in select platforms
+        if sys.platform in ('win32','linux'):
+            action = self.register_tools_menu_action("&Gpu Clocker")
+            action.triggered.connect(launch_sblhelper)
+        else:
+            action = QWidgetAction(None) # dummy action
+
+        self.localsblhelper_launcher_menuitem = action
+
     def register_open_logs_menuitem(self):
         ''' Registers convenience "Open Logs Folder" menu item '''
         action = self.register_tools_menu_action("&Open Logs Folder")
diff --git a/Engine/Plugins/VirtualProduction/Switchboard/Source/Switchboard/switchboard/tools/application_launcher.py b/Engine/Plugins/VirtualProduction/Switchboard/Source/Switchboard/switchboard/tools/application_launcher.py
--- a/Engine/Plugins/VirtualProduction/Switchboard/Source/Switchboard/switchboard/tools/application_launcher.py
+++ b/Engine/Plugins/VirtualProduction/Switchboard/Source/Switchboard/switchboard/tools/application_launcher.py
@@ -16,6 +16,7 @@ class ApplicationLauncher(object):
         self.lock = threading.Lock()
         self.process: Optional[subprocess.Popen] = None
         self.name = name
+        self.requires_elevation = False
 
     def exe_path(self):
         ''' Returns the expected executable path. Extension not included '''
@@ -48,11 +49,11 @@ class ApplicationLauncher(object):
 
             LOGGER.debug(f"Launching '{cmdline}' ...")
 
+            shell = True if self.requires_elevation else False
             creationflags = 0
             if hasattr(subprocess, 'CREATE_NEW_CONSOLE'):
                 creationflags = subprocess.CREATE_NEW_CONSOLE
-
-            subprocess.Popen(cmdline, creationflags=creationflags)
+            subprocess.Popen(cmdline, shell=shell, creationflags=creationflags)
 
         return True
 
diff --git a/Engine/Plugins/VirtualProduction/Switchboard/Source/Switchboard/switchboard/tools/sblhelper_launcher.py b/Engine/Plugins/VirtualProduction/Switchboard/Source/Switchboard/switchboard/tools/sblhelper_launcher.py
new file mode 100644
--- /dev/null
+++ b/Engine/Plugins/VirtualProduction/Switchboard/Source/Switchboard/switchboard/tools/sblhelper_launcher.py
@@ -0,0 +1,18 @@
+# Copyright Epic Games, Inc. All Rights Reserved.
+
+from switchboard.config import CONFIG
+from .application_launcher import ApplicationLauncher
+
+class SBLHelperLauncher(ApplicationLauncher):
+    ''' Manages launching Switchboard Listener Helper'''
+
+    def __init__(self, name:str = 'SwitchboardListenerHelper'):
+        super().__init__(name)
+
+        # The SBL Helper executable requires elevation in order to be able to lock/unlock the gpu clocks.
+        self.requires_elevation = True
+
+    #~Begin ApplicationLauncher interface
+    def exe_path(self):
+        return CONFIG.sblhelper_path()
+    #~End ApplicationLauncher interface
diff --git a/Engine/Source/Programs/SwitchboardListener/Private/SBLHelperClient.cpp b/Engine/Source/Programs/SwitchboardListener/Private/SBLHelperClient.cpp
new file mode 100644
--- /dev/null
+++ b/Engine/Source/Programs/SwitchboardListener/Private/SBLHelperClient.cpp
@@ -0,0 +1,339 @@
+// Copyright Epic Games, Inc. All Rights Reserved.
+
+#include "SBLHelperClient.h"
+
+#include "Common/TcpSocketBuilder.h"
+#include "Dom/JsonValue.h"
+#include "JsonObjectConverter.h"
+#include "Policies/CondensedJsonPrintPolicy.h"
+#include "Serialization/JsonReader.h"
+#include "Serialization/JsonSerializer.h"
+#include "Serialization/JsonWriter.h"
+
+
+
+DEFINE_LOG_CATEGORY(LogSBLHelperClient);
+
+const int32 FSBLHelperClient::CompatibleVersion[3] = { 1, 0, 0 };
+
+FSBLHelperClient::~FSBLHelperClient()
+{
+	Disconnect();
+}
+
+bool FSBLHelperClient::IsConnected() 
+{
+	int32 BytesSent;
+	return Socket && Socket->Send(nullptr, 0, BytesSent);
+}
+
+bool FSBLHelperClient::Connect(const FConnectionParams& InConnectionParams)
+{
+	check(IsInGameThread());
+	
+	// Cache connection parameters
+	ConnectionParams = InConnectionParams;
+
+	// Make sure we have released everything related to any previous connection attempts
+	Disconnect();
+	check(!Socket);
+
+	// Create socket for the new connection
+
+	Socket = FTcpSocketBuilder(TEXT("SBLHelperSocket"));
+	check(Socket);
+
+	// Attempt to connect
+	{
+		const bool bConnected = Socket->Connect(*ConnectionParams.Endpoint.ToInternetAddr());
+
+		if (!bConnected)
+		{
+			Disconnect();
+			return false;
+		}
+	}
+
+	// Confirm that this is actually a SBL server by parsing its first packet(s), if any.
+	constexpr float TimeoutSeconds = 0.5;
+	ParseIncomingMessages(TimeoutSeconds);
+
+	// If the hello packet arrived, ServerVersion must have been populated, which we verify below.
+	{
+		TArray<FString> VersionStringArray;
+		ServerVersion.ParseIntoArray(VersionStringArray, TEXT("."));
+
+		if (VersionStringArray.Num() != 3)
+		{
+			UE_LOG(LogSBLHelperClient, Error, TEXT("Invalid SBLHelperService version '%s' at endpoint %s"), 
+				*ServerVersion,
+				*ConnectionParams.Endpoint.ToString()
+			);
+
+			Disconnect();
+			return false;
+		}
+
+		const int32 VersionArray[3] = {
+			FCString::Atoi(*VersionStringArray[0]),
+			FCString::Atoi(*VersionStringArray[1]),
+			FCString::Atoi(*VersionStringArray[2]),
+		};
+
+		// We'll only look at major and minor.
+		if ((VersionArray[0] != CompatibleVersion[0]) 
+			|| (VersionArray[1] < CompatibleVersion[1]))
+		{
+			UE_LOG(LogSBLHelperClient, Error, TEXT("Incompatible SBLHelperService version '%s' at endpoint %s. Minimum required is %d.%d.%d"), 
+				*ServerVersion,
+				*ConnectionParams.Endpoint.ToString(),
+				CompatibleVersion[0],
+				CompatibleVersion[1],
+				CompatibleVersion[2]
+			);
+
+			Disconnect();
+			return false;
+		}
+	}
+
+	// Ok, we're good to go.
+	UE_LOG(LogSBLHelperClient, Display, TEXT("Connected to SBLHelperService version '%s' at endpoint %s"),
+		*ServerVersion,
+		*ConnectionParams.Endpoint.ToString()
+	);
+
+	return true;
+}
+
+void FSBLHelperClient::HandleCmdHello(TSharedPtr<FJsonObject>& Json)
+{
+	check(Json.IsValid());
+
+	// We use "sblhver" instead of just "ver" as a magic value in disguise.
+	TSharedPtr<FJsonValue> VersionField = Json->TryGetField(TEXT("sblhver"));
+
+	if (!VersionField.IsValid())
+	{
+		return;
+	}
+
+	ServerVersion = VersionField->AsString();
+}
+
+void FSBLHelperClient::ParseMessage(const FString& Message)
+{
+	// All messages are in Json format
+
+	TSharedRef<TJsonReader<TCHAR>> Reader = FJsonStringReader::Create(Message);
+	TSharedPtr<FJsonObject> JsonData;
+
+	if (!FJsonSerializer::Deserialize(Reader, JsonData))
+	{
+		return;
+	}
+
+	// All messages must have a cmd field.
+	TSharedPtr<FJsonValue> CmdField = JsonData->TryGetField(TEXT("cmd"));
+	if (!CmdField.IsValid())
+	{
+		return;
+	}
+
+	const FString Cmd = CmdField->AsString();
+
+	// Each command has its own handler
+	if (Cmd.Equals("hello"))
+	{
+		HandleCmdHello(JsonData);
+	}
+	else
+	{
+		UE_LOG(LogSBLHelperClient, Error, TEXT("Unhandled command '%s'"), *Cmd);
+	}
+}
+
+void FSBLHelperClient::ParseIncomingMessages(const float TimeoutSeconds)
+{
+	for (const FString& Message : ReceiveMessages(TimeoutSeconds))
+	{
+		ParseMessage(Message);
+	}
+}
+
+void FSBLHelperClient::Tick()
+{
+	if (!Socket)
+	{
+		return;
+	}
+
+	// We parse and handle all incoming messages.
+	ParseIncomingMessages();
+}
+
+TArray<FString> FSBLHelperClient::ReceiveMessages(const float TimeoutSeconds)
+{
+	TArray<FString> Messages;
+
+	if (!Socket)
+	{
+		return Messages;
+	}
+
+	const double TimeoutTime = FPlatformTime::Seconds() + TimeoutSeconds;
+
+	do // Execute at least once
+	{
+		uint32 PendingDataSize = 0;
+
+		while (Socket->HasPendingData(PendingDataSize))
+		{
+			// Read 
+			TArray<uint8> Buffer;
+			Buffer.AddUninitialized(PendingDataSize);
+			int32 BytesRead = 0;
+
+			if (!Socket->Recv(Buffer.GetData(), PendingDataSize, BytesRead, ESocketReceiveFlags::None))
+			{
+				UE_LOG(LogSBLHelperClient, Error, TEXT("Error while receiving data via endpoint %s"), 
+					*ConnectionParams.Endpoint.ToString()
+				);
+
+				continue;
+			}
+
+			for (int32 ByteIdx = 0; ByteIdx < BytesRead; ++ByteIdx)
+			{
+				ReceiveBuffer.Add(Buffer[ByteIdx]);
+
+				if (Buffer[ByteIdx] == '\x00')
+				{
+					Messages.Emplace(UTF8_TO_TCHAR(ReceiveBuffer.GetData()));
+
+					// Clear the reception buffer so that it only contains the next partial message
+					ReceiveBuffer.Empty();
+				}
+			}
+		}
+
+		// If we already got at least one message, we can return.
+		if (Messages.Num())
+		{
+			return Messages;
+		}
+
+		// We don't wait more that the given timeout
+		const double CurrentTime = FPlatformTime::Seconds();
+		const double RemainingTimeBeforeTimeout = TimeoutTime - CurrentTime;
+
+		if (RemainingTimeBeforeTimeout <= 0)
+		{
+			break;
+		}
+
+		// Sleep for a bit but not to exceed our timeout time.
+		const double SleepSeconds = FMath::Min(0.025, RemainingTimeBeforeTimeout);
+		FPlatformProcess::SleepNoStats(SleepSeconds);
+
+	} while (true);
+
+	return Messages;
+}
+
+uint32 FSBLHelperClient::GetNextPacketId()
+{
+	NextPacketId++;
+
+	// Avoid 0, which is reserved for no packet id.
+
+	if (NextPacketId == 0)
+	{
+		NextPacketId = 1;
+	}
+
+	return NextPacketId;
+}
+
+FString FSBLHelperClient::CreateMessage(const FString& Cmd, const TMap<FString, FString>& AdditionalFields)
+{
+	FString Message;
+
+	TSharedRef<TJsonWriter<TCHAR, TCondensedJsonPrintPolicy<TCHAR>>> JsonWriter 
+		= TJsonWriterFactory<TCHAR, TCondensedJsonPrintPolicy<TCHAR>>::Create(&Message);
+
+	JsonWriter->WriteObjectStart();
+
+	// message id
+	JsonWriter->WriteValue(TEXT("id"), FString::Printf(TEXT("%u"), GetNextPacketId()));
+
+	// cmd
+	JsonWriter->WriteValue(TEXT("cmd"), Cmd);
+
+	for (const auto& Value : AdditionalFields)
+	{
+		JsonWriter->WriteValue(Value.Key, Value.Value);
+	}
+
+	JsonWriter->WriteObjectEnd();
+	JsonWriter->Close();
+
+	return Message;
+}
+
+void FSBLHelperClient::Disconnect()
+{
+	// Send by message
+	SendMessage(CreateMessage(TEXT("bye"), TMap<FString, FString>()));
+
+	// Clear buffer so that we don't try to use leftover data on the next connection
+	ReceiveBuffer.Empty();
+
+	// Clear the version as well, for the same reason.
+	ServerVersion.Empty();
+
+	// Close and destroy socket
+	if (Socket)
+	{
+		Socket->Close();
+		ISocketSubsystem::Get()->DestroySocket(Socket);
+		Socket = nullptr;
+	}
+}
+
+bool FSBLHelperClient::LockGpuClock(const uint32 Pid)
+{
+	// Send the lock command and the related process pid.
+
+	TMap<FString, FString> AdditionalFields;
+	AdditionalFields.Emplace(TEXT("pid"), FString::Printf(TEXT("%u"), Pid));
+
+	return SendMessage(CreateMessage(TEXT("lock"), AdditionalFields));
+}
+
+bool FSBLHelperClient::SendMessage(const FString& Message)
+{
+	if (!Socket)
+	{
+		return false;
+	}
+
+	UE_LOG(LogSBLHelperClient, Verbose, TEXT("Sending message: %s"), *Message);
+
+	// Send the message over the socket. 
+	// Note the "+ 1" to include the null terminator, which is the message delimiter in the SBLHelper protocol.
+
+	int32 BytesSent = 0;
+	return Socket->Send((uint8*)TCHAR_TO_UTF8(*Message), Message.Len() + 1, BytesSent);
+}
+
+FString FSBLHelperClient::GetStatus()
+{
+	if (!IsConnected())
+	{
+		return TEXT("Disconnected");
+	}
+
+	return TEXT("Connected");
+}
+
diff --git a/Engine/Source/Programs/SwitchboardListener/Private/SBLHelperClient.h b/Engine/Source/Programs/SwitchboardListener/Private/SBLHelperClient.h
new file mode 100644
--- /dev/null
+++ b/Engine/Source/Programs/SwitchboardListener/Private/SBLHelperClient.h
@@ -0,0 +1,124 @@
+// Copyright Epic Games, Inc. All Rights Reserved.
+
+#pragma once
+
+#include "Containers/Array.h"
+#include "Containers/UnrealString.h"
+#include "Interfaces/IPv4/IPv4Endpoint.h"
+#include "Logging/LogMacros.h"
+#include "Templates/SharedPointer.h"
+
+DECLARE_LOG_CATEGORY_EXTERN(LogSBLHelperClient, Log, All);
+
+class FJsonObject;
+class FSocket;
+
+/** Class in charge of communicating with the SwitchboardListenerHelper server */
+class FSBLHelperClient
+{
+public:
+	
+	/** Structure holding the desired connection parameters, such as IP and Port */
+	struct FConnectionParams
+	{
+		/** Server end point to connect to */
+		FIPv4Endpoint Endpoint;
+	};
+
+	/** Destructor */
+	~FSBLHelperClient();
+
+	/** Returns true if currently connected to the SBLHelper server */
+	bool IsConnected();
+
+	/** 
+	 * Connects to the SBLHelper server specified in the given connection parameters. 
+	 * If already connected, it will disconnect before attempting the new connection.
+	 * 
+	 * @return True if it was able to establish the connection to a valid server.
+	 */
+	bool Connect(const FConnectionParams& InConnectionParams);
+
+	/** Disconnects from the SBLHelper server. Can be called even if already disconnected. */
+	void Disconnect();
+
+	/** 
+	 * Requests the server to lock the gpu and memory clocks to the maximum allowed. 
+	 * 
+	 * @param Pid Proces ID to which the lock gpu clock request is tied to.
+	 * 
+	 * @return True if the message was sent over the socket. 
+	 */
+	bool LockGpuClock(const uint32 Pid);
+
+	/** Returns a string describing the status of the server */
+	FString GetStatus();
+
+	/** Call this periodically for tasks such as parsing the incoming messages from the server */
+	void Tick();
+
+private:
+
+	/** 
+	 * Receives messages from the socket into an array. 
+	 * 
+	 * @param TimeoutSeconds Maximum time to wait for at least one message to be completely received.
+	 * 
+	 * @return Array of received messages.
+	 */
+	TArray<FString> ReceiveMessages(const float TimeoutSeconds = 0);
+
+	/** 
+	 * Parses the given message and calls the respective handlers 
+	 * 
+	 * @param InMessage The message to parse.
+	 */
+	void ParseMessage(const FString& InMessage);
+
+	/** 
+	 * Checks the socket for new messages and calls ParseMessage on each.
+	 * 
+	 * @param TimeoutSeconds The maximum time to wait for at least one message to be completely received.
+	 */
+	void ParseIncomingMessages(const float TimeoutSeconds = 0);
+
+	/** Handles the 'Hello' command from the server */
+	void HandleCmdHello(TSharedPtr<FJsonObject>& Json);
+
+	/** 
+	 * Creates a new message for the given command and additional fields 
+	 * 
+	 * @param Cmd The command to create this message for.
+	 * @param AdditionalFields Any additional fields or parameters to include in the message.
+	 * 
+	 * @return The created message, ready to be sent over the socket.
+	 */
+	FString CreateMessage(const FString& Cmd, const TMap<FString, FString>& AdditionalFields);
+
+	/** Sends the given message over the socket, if connected. */
+	bool SendMessage(const FString& Message);
+
+	/** Gets the next packet id to use when creating a new message. Monotonically increments. Skips zero. */
+	uint32 GetNextPacketId();
+
+private:
+
+	/** The minimum compatible SBLHelper server version.  */
+	static const int32 CompatibleVersion[3];
+
+	/** Socket for this client to communicate with the server */
+	FSocket* Socket = nullptr;
+
+	/** Next packet id when creating a new message to be sent */
+	uint32 NextPacketId = 1;
+
+	/** Buffer to hold partial messages received over the socket from the server */
+	TArray<uint8> ReceiveBuffer;
+
+	/** Cached connection parameters. Mostly used for logging purposes. */
+	FConnectionParams ConnectionParams;
+
+	/** The current server version that this client is connected to. */
+	FString ServerVersion;
+};
+
diff --git a/Engine/Source/Programs/SwitchboardListener/Private/SwitchboardListener.cpp b/Engine/Source/Programs/SwitchboardListener/Private/SwitchboardListener.cpp
--- a/Engine/Source/Programs/SwitchboardListener/Private/SwitchboardListener.cpp
+++ b/Engine/Source/Programs/SwitchboardListener/Private/SwitchboardListener.cpp
@@ -3,6 +3,7 @@
 #include "SwitchboardListener.h"
 
 #include "CpuUtilizationMonitor.h"
+#include "SBLHelperClient.h"
 #include "SwitchboardListenerApp.h"
 #include "SwitchboardMessageFuture.h"
 #include "SwitchboardPacket.h"
@@ -107,6 +108,7 @@ struct FRunningProcess
 
 	std::atomic<bool> bPendingKill = false;
 	bool bUpdateClientsWithStdout = false;
+	bool bLockGpuClock = false;
 
 	bool CreatePipes()
 	{
@@ -217,6 +219,7 @@ FSwitchboardListener::FSwitchboardListener(const FSwitchboardCommandLineOptions&
 	: Options(InOptions)
 	, SocketListener(nullptr)
 	, CpuMonitor(MakeShared<FCpuUtilizationMonitor>())
+	, SBLHelper(MakeShared<FSBLHelperClient>())
 	, bIsNvAPIInitialized(false)
 	, CachedMosaicToposLock(MakeShared<FRWLock>())
 	, CachedMosaicTopos(MakeShared<TArray<FMosaicTopo>>())
@@ -378,6 +381,7 @@ bool FSwitchboardListener::Tick()
 	HandleRunningProcesses(RunningProcesses, true);
 	HandleRunningProcesses(FlipModeMonitors, false);
 	SendMessageFutures();
+	SBLHelper->Tick();
 
 	if (IsEngineExitRequested())
 	{
@@ -646,6 +650,7 @@ bool FSwitchboardListener::Task_StartProcess(const FSwitchboardStartTask& InRunT
 	NewProcess->Name = InRunTask.Name;
 	NewProcess->Caller = InRunTask.Caller;
 	NewProcess->bUpdateClientsWithStdout = InRunTask.bUpdateClientsWithStdout;
+	NewProcess->bLockGpuClock = InRunTask.bLockGpuClock;
 	NewProcess->UUID = InRunTask.TaskID; // Process ID is the same as the message ID.
 	NewProcess->PID = 0; // default value
 
@@ -700,6 +705,55 @@ bool FSwitchboardListener::Task_StartProcess(const FSwitchboardStartTask& InRunT
 		return false;
 	}
 
+	// Lock Gpu Clocks for the lifetime of this PID, if requested
+	if (InRunTask.bLockGpuClock && SBLHelper.IsValid())
+	{
+		// Try to connect to the SBLHelper server if we haven't already
+		if (!SBLHelper->IsConnected())
+		{
+			FSBLHelperClient::FConnectionParams ConnectionParams;
+
+			uint16 Port = 8010; // Default tcp port
+
+			// Apply command line port number override, if present
+			{
+				static uint16 CmdLinePortOverride = 0;
+				static bool bCmdLinePortOverrideParsed = false;
+				static bool bCmdLinePortOverrideValid = false;
+
+				if (!bCmdLinePortOverrideParsed)
+				{
+					bCmdLinePortOverrideParsed = true;
+
+					bCmdLinePortOverrideValid = FParse::Value(FCommandLine::Get(), TEXT("sblhport="), CmdLinePortOverride);
+				}
+
+				if (bCmdLinePortOverrideValid)
+				{
+					Port = CmdLinePortOverride;
+				}
+			}
+
+			const FString HostName = FString::Printf(TEXT("localhost:%d"), Port);
+			FIPv4Endpoint::FromHostAndPort(*HostName, ConnectionParams.Endpoint);
+
+			SBLHelper->Connect(ConnectionParams);
+		}
+
+		if (SBLHelper->IsConnected())
+		{
+			const bool bSentMessage = SBLHelper->LockGpuClock(NewProcess->PID);
+
+			if (!bSentMessage)
+			{
+				UE_LOG(LogSwitchboard, Error, TEXT("Failed to send message to SBLHelper server to request gpu clock locking"));
+			}
+
+			// We disconnect right away because launches happen only far and in between.
+			SBLHelper->Disconnect();
+		}
+	}
+
 	UE_LOG(LogSwitchboard, Display, TEXT("Started process %d: %s %s"), NewProcess->PID, *InRunTask.Command, *InRunTask.Arguments);
 
 	RunningProcesses.Add(NewProcess);
@@ -1605,6 +1659,7 @@ FRunningProcess* FSwitchboardListener::FindOrStartFlipModeMonitorForUUID(const F
 	// The monitor auto-closes when monitored program closes.
 	MonitorProcess->UUID = Process->UUID;
 	MonitorProcess->bUpdateClientsWithStdout = false;
+	MonitorProcess->bLockGpuClock = false;
 	MonitorProcess->Recipient = InvalidEndpoint;
 	MonitorProcess->Name = TEXT("flipmode_monitor");
 
diff --git a/Engine/Source/Programs/SwitchboardListener/Private/SwitchboardListener.h b/Engine/Source/Programs/SwitchboardListener/Private/SwitchboardListener.h
--- a/Engine/Source/Programs/SwitchboardListener/Private/SwitchboardListener.h
+++ b/Engine/Source/Programs/SwitchboardListener/Private/SwitchboardListener.h
@@ -29,6 +29,7 @@ struct FSwitchboardFreeListenerBinaryTask;
 
 class FInternetAddr;
 class FSocket;
+class FSBLHelperClient;
 class FTcpListener;
 
 
@@ -164,6 +165,9 @@ private:
 	TArray<FSwitchboardMessageFuture> MessagesFutures;
 	TSharedPtr<FCpuUtilizationMonitor> CpuMonitor;
 
+	/** Client interface to the Switchboard Listener Helper external process */
+	TSharedPtr<FSBLHelperClient> SBLHelper;
+
 	bool bIsNvAPIInitialized;
 
 	TSharedPtr<FRWLock> CachedMosaicToposLock;
diff --git a/Engine/Source/Programs/SwitchboardListener/Private/SwitchboardListenerVersion.h b/Engine/Source/Programs/SwitchboardListener/Private/SwitchboardListenerVersion.h
--- a/Engine/Source/Programs/SwitchboardListener/Private/SwitchboardListenerVersion.h
+++ b/Engine/Source/Programs/SwitchboardListener/Private/SwitchboardListenerVersion.h
@@ -3,5 +3,5 @@
 #pragma once
 
 #define SBLISTENER_VERSION_MAJOR 2
-#define SBLISTENER_VERSION_MINOR 2
-#define SBLISTENER_VERSION_PATCH 1
+#define SBLISTENER_VERSION_MINOR 3
+#define SBLISTENER_VERSION_PATCH 0
diff --git a/Engine/Source/Programs/SwitchboardListener/Private/SwitchboardProtocol.cpp b/Engine/Source/Programs/SwitchboardListener/Private/SwitchboardProtocol.cpp
--- a/Engine/Source/Programs/SwitchboardListener/Private/SwitchboardProtocol.cpp
+++ b/Engine/Source/Programs/SwitchboardListener/Private/SwitchboardProtocol.cpp
@@ -176,6 +176,11 @@ bool CreateTaskFromCommand(const FString& InCommand, const FIPv4Endpoint& InEndp
 			UpdateClientsWithStdoutField->TryGetBool(Task->bUpdateClientsWithStdout);
 		}
 
+		if (TSharedPtr<FJsonValue> LockGpuClock = JsonData->TryGetField(TEXT("bLockGpuClock")))
+		{
+			LockGpuClock->TryGetBool(Task->bLockGpuClock);
+		}
+
 		if (TSharedPtr<FJsonValue> PriorityModifierField = JsonData->TryGetField(TEXT("priority_modifier")))
 		{
 			PriorityModifierField->TryGetNumber(Task->PriorityModifier);
diff --git a/Engine/Source/Programs/SwitchboardListener/Private/SwitchboardTasks.h b/Engine/Source/Programs/SwitchboardListener/Private/SwitchboardTasks.h
--- a/Engine/Source/Programs/SwitchboardListener/Private/SwitchboardTasks.h
+++ b/Engine/Source/Programs/SwitchboardListener/Private/SwitchboardTasks.h
@@ -108,6 +108,7 @@ struct FSwitchboardStartTask : public FSwitchboardTask
 	FString Caller;
 	FString WorkingDir;
 	bool bUpdateClientsWithStdout = false;
+	bool bLockGpuClock = false;
 	int32 PriorityModifier = 0;
 
 	//~ Begin FSwitchboardTask interface
diff --git a/Engine/Source/Programs/SwitchboardListener/SwitchboardListener.Target.cs b/Engine/Source/Programs/SwitchboardListener/SwitchboardListener.Target.cs
--- a/Engine/Source/Programs/SwitchboardListener/SwitchboardListener.Target.cs
+++ b/Engine/Source/Programs/SwitchboardListener/SwitchboardListener.Target.cs
@@ -11,6 +11,8 @@ public class SwitchboardListenerTarget : TargetRules
 		LinkType = TargetLinkType.Monolithic;
 		LaunchModuleName = "SwitchboardListener";
 
+		IncludeOrderVersion = EngineIncludeOrderVersion.Latest;
+
 		// This app compiles against Core/CoreUObject, but not the Engine or Editor, so compile out Engine and Editor references from Core/CoreUObject
 		bCompileAgainstCoreUObject = true;
 		bCompileAgainstEngine = false;
diff --git a/Engine/Source/Programs/SwitchboardListenerHelper/NvmlWrapper/Public/GpuClocker.h b/Engine/Source/Programs/SwitchboardListenerHelper/NvmlWrapper/Public/GpuClocker.h
new file mode 100644
--- /dev/null
+++ b/Engine/Source/Programs/SwitchboardListenerHelper/NvmlWrapper/Public/GpuClocker.h
@@ -0,0 +1,26 @@
+// Copyright Epic Games, Inc. All Rights Reserved.
+
+#pragma once
+
+#include <stdint.h>
+
+#include "NvmlWrapperPublic.h"
+
+/** Encapsulates the logic to lock/unlock gpu clocks */
+class EXPORTLIB FGpuClocker
+{
+public:
+
+	/** Locks the gpu clocks to the maximum allowed */
+	bool LockGpuClocks();
+
+	/** Unlocks the gpu and memory clocks */
+	bool UnlockGpuClocks();
+
+	/** Locks the given gpu and memory clocks to the maximum allowed */
+	bool GpuUnlockDevice(const uint32_t GpuIdx);
+
+	/** Unlocks the given gpu and memory clocks */
+	bool GpuLockDevice(const uint32_t GpuIdx);
+};
+
diff --git a/Engine/Source/Programs/SwitchboardListenerHelper/NvmlWrapper/Public/NvmlWrapperPublic.h b/Engine/Source/Programs/SwitchboardListenerHelper/NvmlWrapper/Public/NvmlWrapperPublic.h
new file mode 100644
--- /dev/null
+++ b/Engine/Source/Programs/SwitchboardListenerHelper/NvmlWrapper/Public/NvmlWrapperPublic.h
@@ -0,0 +1,14 @@
+// Copyright Epic Games, Inc. All Rights Reserved.
+
+#pragma once
+
+#if defined(_MSC_VER)
+	#define EXPORTLIB __declspec(dllexport)
+	#define IMPORTLIB __declspec(dllimport)
+#elif defined(__GNUC__)
+	#define EXPORTLIB __attribute__((visibility("default")))
+	#define IMPORTLIB
+#else
+	#define EXPORTLIB
+	#define IMPORTLIB
+#endif
diff --git a/Engine/Source/Programs/SwitchboardListenerHelper/Private/SBLHelperService.cpp b/Engine/Source/Programs/SwitchboardListenerHelper/Private/SBLHelperService.cpp
new file mode 100644
--- /dev/null
+++ b/Engine/Source/Programs/SwitchboardListenerHelper/Private/SBLHelperService.cpp
@@ -0,0 +1,374 @@
+// Copyright Epic Games, Inc. All Rights Reserved.
+
+#include "SBLHelperService.h"
+#include "SwitchboardListenerHelper.h"
+
+#include "Common/TcpListener.h"
+#include "Tasks/Task.h"
+
+#include "Dom/JsonValue.h"
+#include "JsonObjectConverter.h"
+#include "Policies/CondensedJsonPrintPolicy.h"
+#include "Serialization/JsonReader.h"
+#include "Serialization/JsonSerializer.h"
+#include "Serialization/JsonWriter.h"
+
+
+namespace UE::SBLHelperService
+{
+	/** LAN broadcast address.*/
+	static const FIPv4Endpoint InvalidEndpoint(FIPv4Address::LanBroadcast, 0);
+
+	/** Default timeout for inactive clients before considering them stale */
+	const float DefaultInactiveTimeoutSeconds = 5;
+}
+
+FSBLHelperService::~FSBLHelperService()
+{
+	Stop();
+}
+
+void FSBLHelperService::Start(uint16 Port)
+{
+	// Start from scratch every time.
+	Stop();
+
+	bIsRunning = true;
+
+	FIPv4Endpoint Endpoint;
+	const FString HostName = FString::Printf(TEXT("localhost:%d"), Port);
+	FIPv4Endpoint::FromHostAndPort(*HostName, Endpoint);
+
+	SocketListener = MakeUnique<FTcpListener>(Endpoint, FTimespan::FromSeconds(1), false /*bInReusable*/);
+
+	if (SocketListener->IsActive())
+	{
+		SocketListener->OnConnectionAccepted().BindRaw(this, &FSBLHelperService::OnIncomingConnection);
+		UE_LOG(LogSwitchboardListenerHelper, Display, TEXT("Started listening on %s:%d"), *SocketListener->GetLocalEndpoint().Address.ToString(), SocketListener->GetLocalEndpoint().Port);
+	}
+}
+
+bool FSBLHelperService::OnIncomingConnection(FSocket* InSocket, const FIPv4Endpoint& InEndpoint)
+{
+	UE_LOG(LogSwitchboardListenerHelper, Display, TEXT("Incoming connection via %s:%d"), *InEndpoint.Address.ToString(), InEndpoint.Port);
+
+	PendingConnections.Enqueue(TPair<FIPv4Endpoint, TSharedPtr<FSocket>>(InEndpoint, MakeShareable(InSocket)));
+
+	return true;
+}
+
+void FSBLHelperService::Stop()
+{
+	if (SocketListener.IsValid())
+	{
+		SocketListener->Stop();
+		SocketListener = nullptr;
+	}
+}
+
+bool FSBLHelperService::IsRunning()
+{
+	return bIsRunning;
+}
+
+bool FSBLHelperService::ParseIncomingMessage(const FString& Message, const FIPv4Endpoint& InEndpoint)
+{
+	TSharedRef<TJsonReader<TCHAR>> Reader = FJsonStringReader::Create(Message);
+
+	TSharedPtr<FJsonObject> JsonData;
+
+	if (!FJsonSerializer::Deserialize(Reader, JsonData))
+	{
+		return false;
+	}
+
+	TSharedPtr<FJsonValue> CommandField = JsonData->TryGetField(TEXT("cmd")); // e.g. lock/unlock
+
+	if (!CommandField.IsValid())
+	{
+		return false;
+	}
+
+	const FString Command = CommandField->AsString();
+
+	if (Command.Equals(TEXT("lock")))
+	{
+		return HandleCmdLock(JsonData);
+	}
+	else if (Command.Equals(TEXT("nop")))
+	{
+		return true;
+	}
+	else if (Command.Equals(TEXT("bye")))
+	{
+		// The client wishes to disconnect immediately
+		PendingDisconnections.Enqueue(InEndpoint);
+
+		return true;
+	}
+	else
+	{
+		UE_LOG(LogSwitchboardListenerHelper, Error, TEXT("Unsupported command '%s'"), &Command);
+	}
+
+	return false;
+}
+
+bool FSBLHelperService::HandleCmdLock(const TSharedPtr<FJsonObject>& Json)
+{
+	TSharedPtr<FJsonValue> PidField = Json->TryGetField(TEXT("pid"));
+
+	if (!PidField.IsValid())
+	{
+		UE_LOG(LogSwitchboardListenerHelper, Error, TEXT("Could not get valid pid from lock command"));
+		return false;
+	}
+
+	const uint32 Pid = uint32(FCString::Atoi64(*PidField->AsString()));
+
+	// Get a handle to the process from the Pid
+	FProcHandle ProcHandle = FPlatformProcess::OpenProcess(Pid);
+
+	if (!ProcHandle.IsValid())
+	{
+		UE_LOG(LogSwitchboardListenerHelper, Warning, TEXT("Received Gpu clock lock request for Pid %u, but could not open a handle to it"), &Pid);
+		return false;
+	}
+
+	UE_LOG(LogSwitchboardListenerHelper, Display, TEXT("Received Gpu clock lock request for Pid %u"), &Pid);
+	++NumProcessesHoldingGpuClocks;
+
+	// Now we issue a background task to wait for the process to exit
+	UE::Tasks::Launch(UE_SOURCE_LOCATION, [ProcHandle, this]()
+		{
+			FProcHandle ProcHandleCopy = ProcHandle;
+
+			// This should only return when the process stops.
+			FPlatformProcess::WaitForProc(ProcHandleCopy);
+
+			// Decrement the number of holding gpu clocks.
+			--NumProcessesHoldingGpuClocks;
+		});
+
+	return true;
+}
+
+FString FSBLHelperService::CreateMessage(const FString& Cmd, const TMap<FString, FString>& AdditionalFields)
+{
+	FString Message;
+
+	TSharedRef<TJsonWriter<TCHAR, TCondensedJsonPrintPolicy<TCHAR>>> JsonWriter
+		= TJsonWriterFactory<TCHAR, TCondensedJsonPrintPolicy<TCHAR>>::Create(&Message);
+
+	JsonWriter->WriteObjectStart();
+
+	// message id
+	JsonWriter->WriteValue(TEXT("id"), FString::Printf(TEXT("%u"), GetNextPacketId()));
+
+	// cmd
+	JsonWriter->WriteValue(TEXT("cmd"), Cmd);
+
+	for (const auto& Value : AdditionalFields)
+	{
+		JsonWriter->WriteValue(Value.Key, Value.Value);
+	}
+
+	JsonWriter->WriteObjectEnd();
+	JsonWriter->Close();
+
+	return Message;
+}
+
+uint32 FSBLHelperService::GetNextPacketId()
+{
+	NextPacketId++;
+
+	// Avoid 0, which is reserved for no packet id.
+
+	if (NextPacketId == 0)
+	{
+		NextPacketId = 1;
+	}
+
+	return NextPacketId;
+}
+
+bool FSBLHelperService::SendMessage(const FString& InMessage, const FIPv4Endpoint& InEndpoint)
+{
+	TRACE_CPUPROFILER_EVENT_SCOPE(FSwitchboardListener::SendMessage);
+
+	if (InEndpoint == UE::SBLHelperService::InvalidEndpoint)
+	{
+		return false;
+	}
+
+	if (!Connections.Contains(InEndpoint))
+	{
+		// this happens when a client disconnects while a task it had issued is not finished
+		UE_LOG(LogSwitchboardListenerHelper, Verbose, TEXT("Trying to send message to disconnected client %s"), *InEndpoint.ToString());
+
+		return false;
+	}
+
+	TSharedPtr<FSocket> ClientSocket = Connections[InEndpoint];
+
+	if (!ClientSocket.IsValid())
+	{
+		return false;
+	}
+
+	// Send the message over the socket. 
+	// Note the "+ 1" to include the null terminator, which is the message delimiter in the SBLHelper protocol.
+
+	UE_LOG(LogSwitchboardListenerHelper, Verbose, TEXT("Sending message %s"), *InMessage);
+
+	int32 BytesSent = 0;
+	return ClientSocket->Send((uint8*)TCHAR_TO_UTF8(*InMessage), InMessage.Len() + 1, BytesSent);
+}
+
+void FSBLHelperService::HandlePendingConnections()
+{
+	TPair<FIPv4Endpoint, TSharedPtr<FSocket>> Connection;
+
+	while (PendingConnections.Dequeue(Connection))
+	{
+		Connections.Add(Connection);
+
+		const FIPv4Endpoint ClientEndpoint = Connection.Key;
+
+		LastActivityTime.FindOrAdd(ClientEndpoint, FPlatformTime::Seconds());
+
+		// Send hello message upon connection. The client will be able to verify the version of this server.
+		{
+			const FString Cmd = TEXT("hello");
+
+			TMap<FString, FString> AdditionalFields;
+			AdditionalFields.Emplace(TEXT("sblhver"), FString::Printf(TEXT("%d.%d.%d"),
+				SBLHELPER_VERSION_MAJOR,
+				SBLHELPER_VERSION_MINOR,
+				SBLHELPER_VERSION_PATCH
+			));
+
+			SendMessage(CreateMessage(Cmd, AdditionalFields), ClientEndpoint);
+		}
+	}
+}
+
+void FSBLHelperService::HandlePendingDisconnections()
+{
+	const double CurrentTime = FPlatformTime::Seconds();
+
+	for (const TPair<FIPv4Endpoint, double>& LastActivity : LastActivityTime)
+	{
+		const FIPv4Endpoint& Client = LastActivity.Key;
+		const float ClientTimeoutSeconds = UE::SBLHelperService::DefaultInactiveTimeoutSeconds;
+
+		if ((CurrentTime - LastActivity.Value) > ClientTimeoutSeconds)
+		{
+			UE_LOG(LogSwitchboardListenerHelper, Display, TEXT("Client %s has been inactive for more than %.1fs"), 
+				*Client.ToString(), 
+				ClientTimeoutSeconds
+			);
+
+			PendingDisconnections.Enqueue(Client);
+		}
+	}
+
+	FIPv4Endpoint Client;
+
+	while (PendingDisconnections.Dequeue(Client))
+	{
+		DisconnectClient(Client);
+	}
+}
+
+void FSBLHelperService::DisconnectClient(const FIPv4Endpoint& InClientEndpoint)
+{
+	const FString Client = InClientEndpoint.ToString();
+	UE_LOG(LogSwitchboardListenerHelper, Display, TEXT("Disconnected from Client %s"), *Client);
+
+	Connections.Remove(InClientEndpoint);
+	LastActivityTime.Remove(InClientEndpoint);
+	ReceiveBuffer.Remove(InClientEndpoint);
+}
+
+void FSBLHelperService::ParseIncomingData()
+{
+	for (const TPair<FIPv4Endpoint, TSharedPtr<FSocket>>& Connection : Connections)
+	{
+		const FIPv4Endpoint& ClientEndpoint = Connection.Key;
+		const TSharedPtr<FSocket>& ClientSocket = Connection.Value;
+
+		uint32 PendingDataSize = 0;
+		while (ClientSocket->HasPendingData(PendingDataSize))
+		{
+			LastActivityTime[ClientEndpoint] = FPlatformTime::Seconds();
+
+			TArray<uint8> Buffer;
+			Buffer.AddUninitialized(PendingDataSize);
+
+			int32 BytesRead = 0;
+
+			if (!ClientSocket->Recv(Buffer.GetData(), PendingDataSize, BytesRead, ESocketReceiveFlags::None))
+			{
+				UE_LOG(LogSwitchboardListenerHelper, Error, TEXT("Error while receiving data via endpoint %s"), *ClientEndpoint.ToString());
+				continue;
+			}
+
+			TArray<uint8>& MessageBuffer = ReceiveBuffer.FindOrAdd(ClientEndpoint);
+
+			for (int32 ByteIdx = 0; ByteIdx < BytesRead; ++ByteIdx)
+			{
+				MessageBuffer.Add(Buffer[ByteIdx]);
+
+				// We use '\0' as message terminator
+				if (Buffer[ByteIdx] == '\x00')
+				{
+					const FString Message(UTF8_TO_TCHAR(MessageBuffer.GetData()));
+					ParseIncomingMessage(Message, ClientEndpoint);
+					MessageBuffer.Empty();
+				}
+			}
+		}
+	}
+}
+
+void FSBLHelperService::ManageGpuClocks()
+{
+	if (NumProcessesHoldingGpuClocks)
+	{
+		if (!bGpuClocksLocked)
+		{
+			UE_LOG(LogSwitchboardListenerHelper, Display, TEXT("Locking Gpu Clocks"));
+			GpuClocker.LockGpuClocks();
+			bGpuClocksLocked = true;
+		}
+	}
+	else
+	{
+		if (bGpuClocksLocked)
+		{
+			UE_LOG(LogSwitchboardListenerHelper, Display, TEXT("Unlocking Gpu Clocks"));
+			GpuClocker.UnlockGpuClocks();
+			bGpuClocksLocked = false;
+		}
+	}
+}
+
+void FSBLHelperService::Tick()
+{
+	TRACE_CPUPROFILER_EVENT_SCOPE(FSwitchboardListenerHelper::Tick);
+
+	// Dequeue pending connections
+	HandlePendingConnections();
+
+	// Parse incoming data from remote connections
+	ParseIncomingData();
+
+	// Manage Gpu clocks
+	ManageGpuClocks();
+
+	// Close stale sockets
+	HandlePendingDisconnections();
+}
+
diff --git a/Engine/Source/Programs/SwitchboardListenerHelper/Private/SBLHelperService.h b/Engine/Source/Programs/SwitchboardListenerHelper/Private/SBLHelperService.h
new file mode 100644
--- /dev/null
+++ b/Engine/Source/Programs/SwitchboardListenerHelper/Private/SBLHelperService.h
@@ -0,0 +1,120 @@
+// Copyright Epic Games, Inc. All Rights Reserved.
+
+#pragma once
+
+#include "CoreMinimal.h"
+
+#include "Containers/Set.h"
+#include "HAL/PlatformApplicationMisc.h"
+#include "Interfaces/IPv4/IPv4Endpoint.h"
+
+#include "GpuClocker.h"
+
+class FJsonObject;
+class FTcpListener;
+
+/** Provides services that Switchboard listener cannot perform on its own due to required user privileges */
+class FSBLHelperService
+{
+public:
+
+	~FSBLHelperService();
+
+	/** Starts the service and tcp server */
+	void Start(uint16 Port = 8010);
+
+	/** Stops the service and tcp server */
+	void Stop();
+
+	/** Returns true if the service is currently running or not.*/
+	bool IsRunning();
+	
+	/** Call periodically to perform service maintenance tasks */
+	void Tick();
+
+private:
+
+	/** Called when there is a new tcp client connection */
+	bool OnIncomingConnection(FSocket* InSocket, const FIPv4Endpoint& InEndpoint);
+
+	/** 
+	 * Parses and handles the specified incoming message 
+	 * 
+	 * @param Message The message to parse and handle.
+	 * @param InEndpoint The endpoint that sent the given message.
+	 * 
+	 * @return True if the message was parsed and handled
+	 */
+	bool ParseIncomingMessage(const FString& Message, const FIPv4Endpoint& InEndpoint);
+
+	/**
+	 * Creates a new message for the given command and additional fields
+	 *
+	 * @param Cmd The command to create this message for.
+	 * @param AdditionalFields Any additional fields or parameters to include in the message.
+	 *
+	 * @return The created message, ready to be sent over the socket.
+	 */
+	FString CreateMessage(const FString& Cmd, const TMap<FString, FString>& AdditionalFields);
+
+	/** Gets the next packet id to use when creating a new message. Monotonically increments. Skips zero. */
+	uint32 GetNextPacketId();
+
+	/** Send the given message to the given endpoint */
+	bool SendMessage(const FString& InMessage, const FIPv4Endpoint& InEndpoint);
+
+	/** Handles the lock gpu clock command */
+	bool HandleCmdLock(const TSharedPtr<FJsonObject>& Json);
+
+	/** Dequeues the pending new tcp client connections and sends a hello message */
+	void HandlePendingConnections();
+
+	/** Parses all incoming data from all connected clients */
+	void ParseIncomingData();
+
+	/** Locks/unlocks the gpu clocks depending on the Pids that are holding said request */
+	void ManageGpuClocks();
+
+	/** Removes stale client connections and clients who sent a bye cmd */
+	void HandlePendingDisconnections();
+
+	/** Removes the given client endpoint from the list of active connections and frees related resources */
+	void DisconnectClient(const FIPv4Endpoint& InClientEndpoint);
+
+private:
+
+	/** List of incoming client connections that have not been handled yet */
+	TQueue<TPair<FIPv4Endpoint, TSharedPtr<FSocket>>, EQueueMode::Spsc> PendingConnections;
+
+	/** List of active client connections */
+	TMap<FIPv4Endpoint, TSharedPtr<FSocket>> Connections;
+
+	/** List of client connections to be disconnected (due to inactivity or bye command) */
+	TQueue<FIPv4Endpoint, EQueueMode::Spsc> PendingDisconnections;
+
+	/** Receive buffers dedicated to each client */
+	TMap<FIPv4Endpoint, TArray<uint8>> ReceiveBuffer;
+
+	/** Flag indicating whether the service is currently running or not */
+	bool bIsRunning = false;
+
+	/** Next packet id when creating a new message to be sent */
+	uint32 NextPacketId = 1;
+
+	/** The listening tcp socket */
+	TUniquePtr<FTcpListener> SocketListener;
+
+	/** Count of processes that are holding gpu clock lock active */
+	std::atomic<int32> NumProcessesHoldingGpuClocks = 0;
+
+	/** Whether the gpu clocks have been locked by this service or not */
+	bool bGpuClocksLocked = false;
+
+	/** Last activity time for each client. Used to detect stale connections. */
+	TMap<FIPv4Endpoint, double> LastActivityTime;
+
+	/** Used to lock/unlock the gpu clocks on demand */
+	FGpuClocker GpuClocker;
+
+};
+
diff --git a/Engine/Source/Programs/SwitchboardListenerHelper/Private/SwitchboardListenerHelper.cpp b/Engine/Source/Programs/SwitchboardListenerHelper/Private/SwitchboardListenerHelper.cpp
new file mode 100644
--- /dev/null
+++ b/Engine/Source/Programs/SwitchboardListenerHelper/Private/SwitchboardListenerHelper.cpp
@@ -0,0 +1,220 @@
+// Copyright Epic Games, Inc. All Rights Reserved.
+
+#include "SwitchboardListenerHelper.h"
+
+#include "HAL/ExceptionHandling.h"
+#include "HAL/IConsoleManager.h"
+#include "Interfaces/IPv4/IPv4Endpoint.h"
+#include "Misc/Parse.h"
+#include "Misc/ScopeExit.h"
+#include "ProfilingDebugging/TraceAuxiliary.h"
+#include "RequiredProgramMainCPPInclude.h"
+#include "SBLHelperService.h"
+
+IMPLEMENT_APPLICATION(SwitchboardListenerHelper, "SwitchboardListenerHelper");
+DEFINE_LOG_CATEGORY(LogSwitchboardListenerHelper);
+
+
+namespace UE::SwitchboardListenerHelper
+{
+	/** Initializes the Engine */
+	static int32 InitEngine(const TCHAR* InCommandLine)
+	{
+		const int32 InitResult = GEngineLoop.PreInit(InCommandLine);
+		if (InitResult != 0)
+		{
+			return InitResult;
+		}
+
+		ProcessNewlyLoadedUObjects();
+		FModuleManager::Get().StartProcessingNewlyLoadedObjects();
+
+		// Load internal plugins in the pre-default phase
+		IPluginManager::Get().LoadModulesForEnabledPlugins(ELoadingPhase::PreDefault);
+
+		// Load plugins in default phase
+		IPluginManager::Get().LoadModulesForEnabledPlugins(ELoadingPhase::Default);
+
+		return 0;
+	}
+
+	/** Loads the socket module */
+	static bool InitSocketSystem()
+	{
+		EModuleLoadResult LoadResult;
+		FModuleManager::Get().LoadModuleWithFailureReason(TEXT("Sockets"), LoadResult);
+
+		FIPv4Endpoint::Initialize();
+
+		return LoadResult == EModuleLoadResult::Success;
+	}
+
+	/** Clean teardown of the Engine */
+	static void EngineTearDown()
+	{
+		// Cleanup. See also: FEngineLoop::AppPreExit and FEngineLoop::AppExit
+
+		RequestEngineExit(TEXT("Shutting down SwitchboardListenerHelper"));
+
+		FCoreDelegates::OnPreExit.Broadcast();
+		FCoreDelegates::OnExit.Broadcast();
+		FModuleManager::Get().UnloadModulesAtShutdown();
+
+#if STATS
+		FThreadStats::StopThread();
+#endif
+
+		FTaskGraphInterface::Shutdown();
+		FPlatformMisc::PlatformTearDown();
+
+		if (GLog)
+		{
+			GLog->TearDown();
+		}
+
+		FTextLocalizationManager::TearDown();
+		FInternationalization::TearDown();
+
+		FTraceAuxiliary::Shutdown();
+	}
+
+	/** Called when the window is closed or interrupted */
+	static void HandleApplicationWillTerminate()
+	{
+		RequestEngineExit(TEXT("HandleApplicationWillTerminate"));
+	}
+
+	/** This is the main loop of the program */
+	static void MainLoop()
+	{
+		// We need a chance to handle service being terminated and not leave the GPU clocks locked if not desired.
+		FCoreDelegates::ApplicationWillTerminateDelegate.AddStatic(&HandleApplicationWillTerminate);
+
+		double LastTime = FPlatformTime::Seconds();
+		const float MinFrameTimeSeconds = 0.1;
+
+		// Start the service
+		FSBLHelperService Service;
+		{
+			uint16 Port = 8010; // Default tcp port
+			uint16 PortOverride = Port;
+
+			if (FParse::Value(FCommandLine::Get(), TEXT("sblhport="), PortOverride))
+			{
+				Port = PortOverride;
+			}
+
+			Service.Start(Port);
+		}
+
+		while (Service.IsRunning() && !IsEngineExitRequested())
+		{
+			const double CurrentTime = FPlatformTime::Seconds();
+			const double DeltaTime = CurrentTime - LastTime;
+			LastTime = CurrentTime;
+
+			FTaskGraphInterface::Get().ProcessThreadUntilIdle(ENamedThreads::GameThread);
+
+			// Pump & Tick objects
+			FTSTicker::GetCoreTicker().Tick(DeltaTime);
+
+			// Tick our service
+			Service.Tick();
+
+			GFrameCounter++;
+			FStats::AdvanceFrame(false);
+			GLog->FlushThreadedLogs();
+
+			// Run garbage collection for the UObjects for the rest of the frame or at least to 10 ms
+			IncrementalPurgeGarbage(true, FMath::Max<float>(0.010f, MinFrameTimeSeconds - (FPlatformTime::Seconds() - CurrentTime)));
+
+			// Throttle main thread main fps by sleeping if we still have time
+			FPlatformProcess::Sleep(FMath::Max<float>(0.0f, MinFrameTimeSeconds - (FPlatformTime::Seconds() - CurrentTime)));
+		}
+
+		Service.Stop();
+	}
+
+	/** Initializes the Engine, calls the Main loop, then performs teardown when exiting the application */
+	static int32 RunSwitchboardListenerHelper(int32 ArgC, TCHAR* ArgV[])
+	{
+		// Populate the CommandLine from the input argument.
+		const FString CommandLine = FCommandLine::BuildFromArgV(nullptr, ArgC, ArgV, nullptr);
+		FCommandLine::Set(*CommandLine);
+
+		// Initialize the Engine
+		{
+			const int32 InitResult = InitEngine(*CommandLine);
+
+			UE_LOG(LogSwitchboardListenerHelper, Display, TEXT("SwitchboardListenerHelper %u.%u.%u"), SBLHELPER_VERSION_MAJOR, SBLHELPER_VERSION_MINOR, SBLHELPER_VERSION_PATCH);
+			UE_LOG(LogSwitchboardListenerHelper, Display, TEXT("NOTE: This process should be run with elevated privileges to operate properly."));
+
+			if (InitResult != 0)
+			{
+				UE_LOG(LogSwitchboardListenerHelper, Fatal, TEXT("Could not initialize engine, Error code: %d"), InitResult);
+				RequestEngineExit(TEXT("Engine init failure"));
+				return InitResult;
+			}
+
+			if (!InitSocketSystem())
+			{
+				UE_LOG(LogSwitchboardListenerHelper, Fatal, TEXT("Could not initialize socket system!"));
+				RequestEngineExit(TEXT("Socket init failure"));
+				return 1;
+			}
+		}
+
+#if PLATFORM_WINDOWS
+		// Always minimize on launch
+		ShowWindow(GetConsoleWindow(), SW_MINIMIZE);
+#endif
+
+		// Call main loop of this application
+		MainLoop();
+
+		// Clean tear down of the Engine
+		EngineTearDown();
+
+		return 0;
+	}
+}
+
+/** This is the program's entry point */
+INT32_MAIN_INT32_ARGC_TCHAR_ARGV()
+{
+	using namespace UE::SwitchboardListenerHelper;
+
+	int32 ExitCode;
+	if (FPlatformMisc::IsDebuggerPresent())
+	{
+		ExitCode = RunSwitchboardListenerHelper(ArgC, ArgV);
+	}
+	else
+	{
+#if PLATFORM_WINDOWS && !PLATFORM_SEH_EXCEPTIONS_DISABLED
+		__try
+#endif
+		{
+			// SetCrashHandler(nullptr) sets up default behavior for Linux and Mac interfacing with CrashReportClient
+			FPlatformMisc::SetCrashHandler(nullptr);
+			GIsGuarded = true;
+			ExitCode = RunSwitchboardListenerHelper(ArgC, ArgV);
+			GIsGuarded = false;
+		}
+#if PLATFORM_WINDOWS && !PLATFORM_SEH_EXCEPTIONS_DISABLED
+		__except (ReportCrash(GetExceptionInformation()))
+		{
+			if (GError)
+			{
+				GError->HandleError();
+			}
+
+			// RequestExit(bForce==True) terminates this process with a specific exit code and does not return.
+			FPlatformMisc::RequestExit(true);
+			ExitCode = 1; // Suppress compiler warning.
+		}
+#endif
+	}
+
+	return ExitCode;
+}
diff --git a/Engine/Source/Programs/SwitchboardListenerHelper/Private/SwitchboardListenerHelper.h b/Engine/Source/Programs/SwitchboardListenerHelper/Private/SwitchboardListenerHelper.h
new file mode 100644
--- /dev/null
+++ b/Engine/Source/Programs/SwitchboardListenerHelper/Private/SwitchboardListenerHelper.h
@@ -0,0 +1,12 @@
+// Copyright Epic Games, Inc. All Rights Reserved.
+
+#pragma once
+
+#include "Logging/LogMacros.h"
+
+#define SBLHELPER_VERSION_MAJOR 1
+#define SBLHELPER_VERSION_MINOR 0
+#define SBLHELPER_VERSION_PATCH 0
+
+DECLARE_LOG_CATEGORY_EXTERN(LogSwitchboardListenerHelper, Verbose, All);
+
diff --git a/Engine/Source/Programs/SwitchboardListenerHelper/Resources/Windows/SBLHelper.manifest.xml b/Engine/Source/Programs/SwitchboardListenerHelper/Resources/Windows/SBLHelper.manifest.xml
new file mode 100644
--- /dev/null
+++ b/Engine/Source/Programs/SwitchboardListenerHelper/Resources/Windows/SBLHelper.manifest.xml
@@ -0,0 +1,41 @@
+<?xml version="1.0" encoding="UTF-8" standalone="yes"?>
+<assembly xmlns="urn:schemas-microsoft-com:asm.v1" manifestVersion="1.0" xmlns:asmv3="urn:schemas-microsoft-com:asm.v3">
+
+  <!-- Tell Windows we want to run as administrator -->
+  <trustInfo xmlns="urn:schemas-microsoft-com:asm.v2">
+    <security>
+      <requestedPrivileges>
+        <requestedExecutionLevel level="requireAdministrator" uiAccess="false" />
+      </requestedPrivileges>
+    </security>
+  </trustInfo>
+
+  <!-- Use themed common controls -->
+  <dependency>
+    <dependentAssembly>
+      <assemblyIdentity type='win32' name='Microsoft.Windows.Common-Controls' version='6.0.0.0' processorArchitecture='amd64' publicKeyToken='6595b64144ccf1df' />
+    </dependentAssembly>
+  </dependency>
+
+  <!-- OS Support -->
+  <compatibility xmlns="urn:schemas-microsoft-com:compatibility.v1">
+    <application>
+      <!-- Windows 10 -->
+      <supportedOS Id="{8e0f7a12-bfb3-4fe8-b9a5-48fd50a15a9a}"/>
+      <!-- Windows 8.1 -->
+      <supportedOS Id="{1f676c76-80e1-4239-95bb-83d0f6d0da78}"/>
+      <!-- Windows 7 -->
+      <supportedOS Id="{35138b9a-5d96-4fbd-8e2d-a2440225f93a}"/>
+      <!-- Windows 8 -->
+      <supportedOS Id="{4a2f28e3-53b9-4441-ba9c-d69d4a4a6e38}"/>
+    </application>
+  </compatibility>
+
+  <!-- Allow paths longer than MAX_PATH (also requires enabling per-system via group policy or registry entry) -->
+  <asmv3:application>
+    <windowsSettings xmlns:ws2="http://schemas.microsoft.com/SMI/2016/WindowsSettings">
+      <ws2:longPathAware>true</ws2:longPathAware>
+    </windowsSettings>
+  </asmv3:application>
+
+</assembly>
\ No newline at end of file
diff --git a/Engine/Source/Programs/SwitchboardListenerHelper/SwitchboardListenerHelper.Build.cs b/Engine/Source/Programs/SwitchboardListenerHelper/SwitchboardListenerHelper.Build.cs
new file mode 100644
--- /dev/null
+++ b/Engine/Source/Programs/SwitchboardListenerHelper/SwitchboardListenerHelper.Build.cs
@@ -0,0 +1,57 @@
+// Copyright Epic Games, Inc. All Rights Reserved.
+
+using System.IO;
+using UnrealBuildTool;
+
+public class SwitchboardListenerHelper : ModuleRules
+{
+	public SwitchboardListenerHelper(ReadOnlyTargetRules Target) : base(Target)
+	{
+		PrivateDependencyModuleNames.AddRange(
+			new string[] {
+				"ApplicationCore", // for LaunchEngineLoop.cpp dependency
+				"Core",
+				"CoreUObject",
+				"Json",
+				"Networking",
+				"Projects", // for LaunchEngineLoop.cpp dependency
+				"JsonUtilities",
+				"TraceLog",
+			}
+		);
+
+		if (Target.Platform == UnrealTargetPlatform.Win64)
+		{
+			PublicAdditionalLibraries.Add(Path.Combine(ModuleDirectory, "NvmlWrapper", "Lib", "Windows", "NvmlWrapper.lib"));
+		}
+		else if (Target.Platform == UnrealTargetPlatform.Linux)
+		{
+			PublicAdditionalLibraries.Add(Path.Combine(ModuleDirectory, "NvmlWrapper", "Lib", "Linux", "libNvmlWrapper.a"));
+		}
+
+		PrivateIncludePathModuleNames.AddRange(
+			new string[] {
+				"Launch",
+			}
+		);
+
+		PrivateIncludePaths.AddRange(
+			new string[] {
+				"Runtime/Launch/Private", // for LaunchEngineLoop.cpp include
+			}
+		);
+
+		PublicIncludePaths.AddRange(
+			new string[]
+			{
+				Path.Combine(ModuleDirectory, "NvmlWrapper", "Public"),
+			}
+		);
+
+		DynamicallyLoadedModuleNames.AddRange(
+			new string[] {
+				"Sockets",
+			}
+		);
+	}
+}
diff --git a/Engine/Source/Programs/SwitchboardListenerHelper/SwitchboardListenerHelper.Target.cs b/Engine/Source/Programs/SwitchboardListenerHelper/SwitchboardListenerHelper.Target.cs
new file mode 100644
--- /dev/null
+++ b/Engine/Source/Programs/SwitchboardListenerHelper/SwitchboardListenerHelper.Target.cs
@@ -0,0 +1,41 @@
+// Copyright Epic Games, Inc. All Rights Reserved.
+
+using UnrealBuildTool;
+using System.Collections.Generic;
+using EpicGames.Core;
+using UnrealBuildBase;
+
+public class SwitchboardListenerHelperTarget : TargetRules
+{
+	public SwitchboardListenerHelperTarget(TargetInfo Target) : base(Target)
+	{
+		Type = TargetType.Program;
+		LinkType = TargetLinkType.Monolithic;
+		LaunchModuleName = "SwitchboardListenerHelper";
+
+		IncludeOrderVersion = EngineIncludeOrderVersion.Latest;
+
+		// This app compiles against Core/CoreUObject, but not the Engine or Editor, so compile out Engine and Editor references from Core/CoreUObject
+		bCompileAgainstCoreUObject = true;
+		bCompileAgainstEngine = false;
+		bBuildWithEditorOnlyData = false;
+
+		// Disable internationalization support, so the exe can be launched outside of the engine directory.
+		bCompileICU = false;
+
+		// The listener is meant to be a console application (no window), but on MacOS, to get a proper log console, a full application must be built.
+		bIsBuildingConsoleApplication = Target.Platform != UnrealTargetPlatform.Mac;
+
+		GlobalDefinitions.Add("UE_TRACE_ENABLED=1");
+
+		if (Target.Platform == UnrealTargetPlatform.Win64)
+		{
+			// Run with elevated privileges.
+
+			AdditionalLinkerArguments = "/MANIFESTUAC:NO";
+			WindowsPlatform.ManifestFile = FileReference.Combine(Unreal.EngineDirectory,
+				"Source", "Programs", "SwitchboardListenerHelper", "Resources", "Windows",
+				"SBLHelper.manifest.xml").ToString();
+		}
+	}
+}
