Add missing dependency for VS Builds.

diff --git a/Engine/Source/Programs/SwitchboardListenerHelper/Private/SwitchboardListenerHelper.cpp b/Engine/Source/Programs/SwitchboardListenerHelper/Private/SwitchboardListenerHelper.cpp
--- a/Engine/Source/Programs/SwitchboardListenerHelper/Private/SwitchboardListenerHelper.cpp
+++ b/Engine/Source/Programs/SwitchboardListenerHelper/Private/SwitchboardListenerHelper.cpp
@@ -2,6 +2,8 @@
 
 #include "SwitchboardListenerHelper.h"
 
+
+#include "Common/TcpListener.h"
 #include "HAL/ExceptionHandling.h"
 #include "HAL/IConsoleManager.h"
 #include "Interfaces/IPv4/IPv4Endpoint.h"
