Fix -nopchm compile error

#rnx

diff --git a/Engine/Source/Programs/SwitchboardListenerHelper/Private/SBLHelperService.h b/Engine/Source/Programs/SwitchboardListenerHelper/Private/SBLHelperService.h
--- a/Engine/Source/Programs/SwitchboardListenerHelper/Private/SBLHelperService.h
+++ b/Engine/Source/Programs/SwitchboardListenerHelper/Private/SBLHelperService.h
@@ -4,6 +4,7 @@
 
 #include "CoreMinimal.h"
 
+#include "Containers/Queue.h"
 #include "Containers/Set.h"
 #include "HAL/PlatformApplicationMisc.h"
 #include "Interfaces/IPv4/IPv4Endpoint.h"
