Fixed format mismatches in FSBLHelperService::HandleCmdLock.

diff --git a/Engine/Source/Programs/SwitchboardListenerHelper/Private/SBLHelperService.cpp b/Engine/Source/Programs/SwitchboardListenerHelper/Private/SBLHelperService.cpp
--- a/Engine/Source/Programs/SwitchboardListenerHelper/Private/SBLHelperService.cpp
+++ b/Engine/Source/Programs/SwitchboardListenerHelper/Private/SBLHelperService.cpp
@@ -131,11 +131,11 @@ bool FSBLHelperService::HandleCmdLock(const TSharedPtr<FJsonObject>& Json)
 
 	if (!ProcHandle.IsValid())
 	{
-		UE_LOG(LogSwitchboardListenerHelper, Warning, TEXT("Received Gpu clock lock request for Pid %u, but could not open a handle to it"), &Pid);
+		UE_LOG(LogSwitchboardListenerHelper, Warning, TEXT("Received Gpu clock lock request for Pid %u, but could not open a handle to it"), Pid);
 		return false;
 	}
 
-	UE_LOG(LogSwitchboardListenerHelper, Display, TEXT("Received Gpu clock lock request for Pid %u"), &Pid);
+	UE_LOG(LogSwitchboardListenerHelper, Display, TEXT("Received Gpu clock lock request for Pid %u"), Pid);
 	++NumProcessesHoldingGpuClocks;
 
 	// Now we issue a background task to wait for the process to exit
